Experimental low-level network stats plugin that captures OS-level network throughput.

 - Only implemented on console platforms at the moment
 - Also sends it to Unreal Insights (requires -trace=counters on the command line)

#rnx

diff --git a/Engine/Plugins/Experimental/LowLevelNetTrace/LowLevelNetTrace.uplugin b/Engine/Plugins/Experimental/LowLevelNetTrace/LowLevelNetTrace.uplugin
new file mode 100644
--- /dev/null
+++ b/Engine/Plugins/Experimental/LowLevelNetTrace/LowLevelNetTrace.uplugin
@@ -0,0 +1,24 @@
+{
+	"FileVersion": 1,
+	"Version": 1,
+	"VersionName": "1.0",
+	"FriendlyName": "Low-level network trace Plugin",
+	"Description": "Actively monitors & reports system network throughput.",
+	"Category": "Profiling",
+	"CreatedBy": "Epic Games, Inc.",
+	"CreatedByURL": "https://epicgames.com",
+	"DocsURL": "",
+	"EnabledByDefault": true,
+	"IsBetaVersion": true,
+	"SupportedTargetPlatforms": [],
+	"HasExplicitPlatforms": true,
+	"Modules": [
+		{
+			"Name": "LowLevelNetTrace",
+			"Type": "Runtime",
+			"LoadingPhase": "PostConfigInit",
+			"PlatformAllowList": [],
+			"HasExplicitPlatforms": true
+		}
+	]
+}
diff --git a/Engine/Plugins/Experimental/LowLevelNetTrace/Source/LowLevelNetTrace/LowLevelNetTrace.Build.cs b/Engine/Plugins/Experimental/LowLevelNetTrace/Source/LowLevelNetTrace/LowLevelNetTrace.Build.cs
new file mode 100644
--- /dev/null
+++ b/Engine/Plugins/Experimental/LowLevelNetTrace/Source/LowLevelNetTrace/LowLevelNetTrace.Build.cs
@@ -0,0 +1,14 @@
+// Copyright Epic Games, Inc. All Rights Reserved.
+
+using EpicGames.Core;
+
+namespace UnrealBuildTool.Rules
+{
+	public class LowLevelNetTrace : ModuleRules
+	{
+		public LowLevelNetTrace(ReadOnlyTargetRules Target) : base(Target)
+        {
+			PrivateDependencyModuleNames.Add("Core");
+		}
+	}
+}
diff --git a/Engine/Plugins/Experimental/LowLevelNetTrace/Source/LowLevelNetTrace/Private/ILowLevelNetTrace.h b/Engine/Plugins/Experimental/LowLevelNetTrace/Source/LowLevelNetTrace/Private/ILowLevelNetTrace.h
new file mode 100644
--- /dev/null
+++ b/Engine/Plugins/Experimental/LowLevelNetTrace/Source/LowLevelNetTrace/Private/ILowLevelNetTrace.h
@@ -0,0 +1,18 @@
+// Copyright Epic Games, Inc. All Rights Reserved.
+
+#pragma once
+
+#include "CoreTypes.h"
+#include "LowLevelNetTraceModule.h"
+
+#if UE_LOW_LEVEL_NET_TRACE_ENABLED
+
+class ILowLevelNetTrace
+{
+public:
+	virtual ~ILowLevelNetTrace() = default;
+
+	virtual bool UpdateSnapshot( FLowLevelNetTraceSnapshot& OutSnapshot ) = 0;
+};
+
+#endif // UE_LOW_LEVEL_NET_TRACE_ENABLED
\ No newline at end of file
diff --git a/Engine/Plugins/Experimental/LowLevelNetTrace/Source/LowLevelNetTrace/Private/LowLevelNetTraceModule.cpp b/Engine/Plugins/Experimental/LowLevelNetTrace/Source/LowLevelNetTrace/Private/LowLevelNetTraceModule.cpp
new file mode 100644
--- /dev/null
+++ b/Engine/Plugins/Experimental/LowLevelNetTrace/Source/LowLevelNetTrace/Private/LowLevelNetTraceModule.cpp
@@ -0,0 +1,114 @@
+// Copyright Epic Games, Inc. All Rights Reserved.
+
+#include "LowLevelNetTraceModule.h"
+
+#if UE_LOW_LEVEL_NET_TRACE_ENABLED
+#include "ILowLevelNetTrace.h"
+#include "ProfilingDebugging/CountersTrace.h"
+#include "HAL/IConsoleManager.h"
+#include "HAL/Runnable.h"
+#include "HAL/RunnableThread.h"
+#include "HAL/Event.h"
+#include "HAL/PlatformProcess.h"
+#include "Templates/UniquePtr.h"
+#include COMPILED_PLATFORM_HEADER(LowLevelNetTrace.h)
+#include <atomic>
+
+TRACE_DECLARE_UNCHECKED_FLOAT_COUNTER(LowLevelNetTrace_UploadMbps,   TEXT("LowLevelNetTrace/UploadMbps"));
+TRACE_DECLARE_UNCHECKED_FLOAT_COUNTER(LowLevelNetTrace_DownloadMbps, TEXT("LowLevelNetTrace/DownloadMbps"));
+
+
+class FLowLevelNetTraceModule : public ILowLevelNetTraceModule, public FRunnable
+{
+public:
+	virtual void StartupModule() override
+	{
+		LowLevelNetTrace.Reset(new FPlatformLowLevelNetTrace());
+		UpdateStatistics();
+		StartThread();
+	}
+
+	virtual void ShutdownModule() override
+	{
+		StopThread();
+		LowLevelNetTrace.Reset();
+	}
+
+	virtual bool GetSnapshot( FLowLevelNetTraceSnapshot& OutSnapshot ) override
+	{
+		OutSnapshot = CachedNetworkSnapshot.load();
+		return true;
+	}
+
+private:
+
+	void UpdateStatistics()
+	{
+		// get a new snapshot
+		FLowLevelNetTraceSnapshot NewSnapshot;
+		if (LowLevelNetTrace->UpdateSnapshot(NewSnapshot))
+		{
+			//update trace stats
+			TRACE_COUNTER_SET(LowLevelNetTrace_UploadMbps,    NewSnapshot.UploadMbps);
+			TRACE_COUNTER_SET(LowLevelNetTrace_DownloadMbps,  NewSnapshot.DownloadMbps);
+
+			// update cached snapshot
+			CachedNetworkSnapshot.exchange(NewSnapshot);
+		}
+	}
+
+	void StartThread()
+	{
+		check(!Thread.IsValid());
+		ThreadStopEvent = FPlatformProcess::GetSynchEventFromPool(true);
+		Thread.Reset( FRunnableThread::Create(this, TEXT("LowLevelLowLevelNetTrace"), 0, TPri_Lowest) );
+	}
+
+	void StopThread()
+	{
+		if (Thread.IsValid())
+		{
+			Thread.Reset();
+			FPlatformProcess::ReturnSynchEventToPool(ThreadStopEvent);
+			ThreadStopEvent = nullptr;
+		}
+	}
+
+	// FRunnable interface
+	virtual uint32 Run() override
+	{
+		// update the stats once per second
+		while(!ThreadStopEvent->Wait(1000))
+		{
+			UpdateStatistics();
+		}
+		return 0;
+	}
+
+	virtual void Stop() override
+	{
+		ThreadStopEvent->Trigger();
+	}
+	// end of FRunnable interface
+
+
+	TUniquePtr<FRunnableThread> Thread;
+	FEvent* ThreadStopEvent = nullptr;
+
+	TUniquePtr<ILowLevelNetTrace> LowLevelNetTrace;
+	std::atomic<FLowLevelNetTraceSnapshot> CachedNetworkSnapshot;
+};
+
+#else
+class FLowLevelNetTraceModule : public ILowLevelNetTraceModule
+{
+	virtual bool GetSnapshot( FLowLevelNetTraceSnapshot& OutSnapshot ) override
+	{ 
+		return false; 
+	}
+};
+#endif // UE_LOW_LEVEL_NET_TRACE_ENABLED
+
+
+IMPLEMENT_MODULE(FLowLevelNetTraceModule, LowLevelNetTrace)
+
diff --git a/Engine/Plugins/Experimental/LowLevelNetTrace/Source/LowLevelNetTrace/Public/LowLevelNetTraceModule.h b/Engine/Plugins/Experimental/LowLevelNetTrace/Source/LowLevelNetTrace/Public/LowLevelNetTraceModule.h
new file mode 100644
--- /dev/null
+++ b/Engine/Plugins/Experimental/LowLevelNetTrace/Source/LowLevelNetTrace/Public/LowLevelNetTraceModule.h
@@ -0,0 +1,34 @@
+// Copyright Epic Games, Inc. All Rights Reserved.
+
+#pragma once
+
+#include "Modules/ModuleManager.h"
+#include "Modules/ModuleInterface.h"
+#include "Trace/Config.h"
+
+#if !defined(UE_LOW_LEVEL_NET_TRACE_ENABLED)
+	#define UE_LOW_LEVEL_NET_TRACE_ENABLED ((!UE_BUILD_SHIPPING) && UE_TRACE_ENABLED)
+#endif
+
+struct FLowLevelNetTraceSnapshot
+{
+	double UploadMbps = 0.0;
+	double DownloadMbps = 0.0;
+	double TimeStamp = 0.0;
+};
+
+class ILowLevelNetTraceModule : public IModuleInterface
+{
+public:
+	static inline ILowLevelNetTraceModule& Get()
+	{
+		return FModuleManager::LoadModuleChecked<ILowLevelNetTraceModule>("LowLevelNetTrace");
+	}
+
+	static inline bool IsAvailable()
+	{
+		return FModuleManager::Get().IsModuleLoaded("LowLevelNetTrace");
+	}
+
+	virtual bool GetSnapshot( FLowLevelNetTraceSnapshot& OutSnapshot ) = 0;
+};
diff --git a/Engine/Source/Runtime/StorageServerClientDebug/Private/StorageServerConnectionDebug.cpp b/Engine/Source/Runtime/StorageServerClientDebug/Private/StorageServerConnectionDebug.cpp
--- a/Engine/Source/Runtime/StorageServerClientDebug/Private/StorageServerConnectionDebug.cpp
+++ b/Engine/Source/Runtime/StorageServerClientDebug/Private/StorageServerConnectionDebug.cpp
@@ -250,6 +250,12 @@ public:
 				StopThread();
 				StartTick();
 			});
+
+			// load the low-level network tracing module too, so we get platform bandwidth stats as well
+			if (FModuleManager::Get().ModuleExists(TEXT("LowLevelNetTrace")))
+			{
+				FModuleManager::Get().LoadModule(TEXT("LowLevelNetTrace"));
+			}
 		}
 	}
 
